Fix TP_Puzzle compile error

diff --git a/Templates/TP_Puzzle/Source/TP_Puzzle/TP_PuzzleBlock.cpp b/Templates/TP_Puzzle/Source/TP_Puzzle/TP_PuzzleBlock.cpp
--- a/Templates/TP_Puzzle/Source/TP_Puzzle/TP_PuzzleBlock.cpp
+++ b/Templates/TP_Puzzle/Source/TP_Puzzle/TP_PuzzleBlock.cpp
@@ -39,7 +39,7 @@ ATP_PuzzleBlock::ATP_PuzzleBlock()
 	OrangeMaterial = ConstructorStatics.OrangeMaterial.Get();
 }
 
-void ATP_PuzzleBlock::BlockClicked(UPrimitiveComponent* ClickedComp)
+void ATP_PuzzleBlock::BlockClicked(UPrimitiveComponent* ClickedComp, FKey ButtonClicked)
 {
 	// Check we are not already active
 	if(!bIsActive)
@@ -60,5 +60,5 @@ void ATP_PuzzleBlock::BlockClicked(UPrimitiveComponent* ClickedComp)
 
 void ATP_PuzzleBlock::OnFingerPressedBlock(ETouchIndex::Type FingerIndex, UPrimitiveComponent* TouchedComponent)
 {
-	BlockClicked(TouchedComponent);
+	BlockClicked(TouchedComponent, EKeys::Invalid);
 }
diff --git a/Templates/TP_Puzzle/Source/TP_Puzzle/TP_PuzzleBlock.h b/Templates/TP_Puzzle/Source/TP_Puzzle/TP_PuzzleBlock.h
--- a/Templates/TP_Puzzle/Source/TP_Puzzle/TP_PuzzleBlock.h
+++ b/Templates/TP_Puzzle/Source/TP_Puzzle/TP_PuzzleBlock.h
@@ -33,7 +33,7 @@ public:
 
 	/** Handle the block being clicked */
 	UFUNCTION()
-	void BlockClicked(UPrimitiveComponent* ClickedComp);
+	void BlockClicked(UPrimitiveComponent* ClickedComp, FKey ButtonClicked);
 
 	/** Handle the block being touched  */
 	UFUNCTION()
