Fixing spelling error in comment.

diff --git a/Templates/TP_Puzzle/Source/TP_Puzzle/TP_PuzzleBlockGrid.cpp b/Templates/TP_Puzzle/Source/TP_Puzzle/TP_PuzzleBlockGrid.cpp
--- a/Templates/TP_Puzzle/Source/TP_Puzzle/TP_PuzzleBlockGrid.cpp
+++ b/Templates/TP_Puzzle/Source/TP_Puzzle/TP_PuzzleBlockGrid.cpp
@@ -39,7 +39,7 @@ void ATP_PuzzleBlockGrid::BeginPlay()
 		const float XOffset = (BlockIndex/Size) * BlockSpacing; // Divide by dimension
 		const float YOffset = (BlockIndex%Size) * BlockSpacing; // Modulo gives remainder
 
-		// Make postion vector, offset from Grid location
+		// Make position vector, offset from Grid location
 		const FVector BlockLocation = FVector(XOffset, YOffset, 0.f) + GetActorLocation();
 
 		// Spawn a block
