Fix input mapping issue in Templates\TP_Puzzle

diff --git a/Templates/TP_Puzzle/Source/TP_Puzzle/TP_PuzzlePawn.cpp b/Templates/TP_Puzzle/Source/TP_Puzzle/TP_PuzzlePawn.cpp
--- a/Templates/TP_Puzzle/Source/TP_Puzzle/TP_PuzzlePawn.cpp
+++ b/Templates/TP_Puzzle/Source/TP_Puzzle/TP_PuzzlePawn.cpp
@@ -43,7 +43,7 @@ void ATP_PuzzlePawn::SetupPlayerInputComponent(UInputComponent* PlayerInputCompo
 {
 	Super::SetupPlayerInputComponent(PlayerInputComponent);
 
-	PlayerInputComponent->BindAction("OnResetVR", EInputEvent::IE_Pressed, this, &ATP_PuzzlePawn::OnResetVR);
+	PlayerInputComponent->BindAction("ResetVR", EInputEvent::IE_Pressed, this, &ATP_PuzzlePawn::OnResetVR);
 	PlayerInputComponent->BindAction("TriggerClick", EInputEvent::IE_Pressed, this, &ATP_PuzzlePawn::TriggerClick);
 }
 
