Fixing puzzle template build.

diff --git a/Templates/TP_Puzzle/Source/TP_Puzzle/TP_PuzzlePawn.cpp b/Templates/TP_Puzzle/Source/TP_Puzzle/TP_PuzzlePawn.cpp
--- a/Templates/TP_Puzzle/Source/TP_Puzzle/TP_PuzzlePawn.cpp
+++ b/Templates/TP_Puzzle/Source/TP_Puzzle/TP_PuzzlePawn.cpp
@@ -76,15 +76,15 @@ void ATP_PuzzlePawn::TraceForBlock(const FVector& Start, const FVector& End, boo
 		DrawDebugLine(GetWorld(), Start, HitResult.Location, FColor::Red);
 		DrawDebugSolidBox(GetWorld(), HitResult.Location, FVector(20.0f), FColor::Red);
 	}
-	if (HitResult.Actor.IsValid())
+	if (HitResult.GetHitObjectHandle().IsValid())
 	{
-		ATP_PuzzleBlock* HitBlock = Cast<ATP_PuzzleBlock>(HitResult.Actor.Get());
-		if (CurrentBlockFocus != HitBlock)
+		if (HitResult.GetHitObjectHandle() != CurrentBlockFocus)
 		{
 			if (CurrentBlockFocus)
 			{
 				CurrentBlockFocus->Highlight(false);
 			}
+			ATP_PuzzleBlock* HitBlock = HitResult.GetHitObjectHandle().FetchActor<ATP_PuzzleBlock>();
 			if (HitBlock)
 			{
 				HitBlock->Highlight(true);
