UE-181824 Move HeadMountedDisplay files to XRBase plugin to reduce UE overhead. - fix TP_PuzzlePawn build error
-Missed one HeadMountedDisplayFunctionLibrary reference when removing VR from this deprecated template.

diff --git a/Templates/TP_Puzzle/Source/TP_Puzzle/TP_PuzzlePawn.cpp b/Templates/TP_Puzzle/Source/TP_Puzzle/TP_PuzzlePawn.cpp
--- a/Templates/TP_Puzzle/Source/TP_Puzzle/TP_PuzzlePawn.cpp
+++ b/Templates/TP_Puzzle/Source/TP_Puzzle/TP_PuzzlePawn.cpp
@@ -43,7 +43,7 @@ void ATP_PuzzlePawn::CalcCamera(float DeltaTime, struct FMinimalViewInfo& OutRes
 
 void ATP_PuzzlePawn::OnResetVR()
 {
-	UHeadMountedDisplayFunctionLibrary::ResetOrientationAndPosition();
+	//UHeadMountedDisplayFunctionLibrary::ResetOrientationAndPosition();
 }
 
 void ATP_PuzzlePawn::TriggerClick()
